[Backout] - CL46390130
#rnx
Original CL Desc
-----------------------------------------------------------------
InEditorDocumentation Plugin

Add plugin source as experimental to UE 5.7
Enable the plugin in the TP_UEIntro_BP template project.

diff --git a/Engine/Plugins/Experimental/InEditorDocumentation/InEditorDocumentation.uplugin b/Engine/Plugins/Experimental/InEditorDocumentation/InEditorDocumentation.uplugin
deleted file mode 100644
diff --git a/Engine/Plugins/Experimental/InEditorDocumentation/Resources/page-not-found.html b/Engine/Plugins/Experimental/InEditorDocumentation/Resources/page-not-found.html
deleted file mode 100644
diff --git a/Engine/Plugins/Experimental/InEditorDocumentation/Source/ChangeWebBrowserUserAgent/ChangeWebBrowserUserAgent.Build.cs b/Engine/Plugins/Experimental/InEditorDocumentation/Source/ChangeWebBrowserUserAgent/ChangeWebBrowserUserAgent.Build.cs
deleted file mode 100644
diff --git a/Engine/Plugins/Experimental/InEditorDocumentation/Source/ChangeWebBrowserUserAgent/ChangeWebBrowserUserAgent.cpp b/Engine/Plugins/Experimental/InEditorDocumentation/Source/ChangeWebBrowserUserAgent/ChangeWebBrowserUserAgent.cpp
deleted file mode 100644
diff --git a/Engine/Plugins/Experimental/InEditorDocumentation/Source/ChangeWebBrowserUserAgent/ChangeWebBrowserUserAgent.h b/Engine/Plugins/Experimental/InEditorDocumentation/Source/ChangeWebBrowserUserAgent/ChangeWebBrowserUserAgent.h
deleted file mode 100644
diff --git a/Engine/Plugins/Experimental/InEditorDocumentation/Source/InEditorDocumentation/InEditorDocumentation.Build.cs b/Engine/Plugins/Experimental/InEditorDocumentation/Source/InEditorDocumentation/InEditorDocumentation.Build.cs
deleted file mode 100644
diff --git a/Engine/Plugins/Experimental/InEditorDocumentation/Source/InEditorDocumentation/Private/CommandsStyle.cpp b/Engine/Plugins/Experimental/InEditorDocumentation/Source/InEditorDocumentation/Private/CommandsStyle.cpp
deleted file mode 100644
diff --git a/Engine/Plugins/Experimental/InEditorDocumentation/Source/InEditorDocumentation/Private/DocumentationCommands.cpp b/Engine/Plugins/Experimental/InEditorDocumentation/Source/InEditorDocumentation/Private/DocumentationCommands.cpp
deleted file mode 100644
diff --git a/Engine/Plugins/Experimental/InEditorDocumentation/Source/InEditorDocumentation/Private/InEditorDocumentation.cpp b/Engine/Plugins/Experimental/InEditorDocumentation/Source/InEditorDocumentation/Private/InEditorDocumentation.cpp
deleted file mode 100644
diff --git a/Engine/Plugins/Experimental/InEditorDocumentation/Source/InEditorDocumentation/Private/InEditorDocumentationSettings.cpp b/Engine/Plugins/Experimental/InEditorDocumentation/Source/InEditorDocumentation/Private/InEditorDocumentationSettings.cpp
deleted file mode 100644
diff --git a/Engine/Plugins/Experimental/InEditorDocumentation/Source/InEditorDocumentation/Public/CommandsStyle.h b/Engine/Plugins/Experimental/InEditorDocumentation/Source/InEditorDocumentation/Public/CommandsStyle.h
deleted file mode 100644
diff --git a/Engine/Plugins/Experimental/InEditorDocumentation/Source/InEditorDocumentation/Public/DocumentationCommands.h b/Engine/Plugins/Experimental/InEditorDocumentation/Source/InEditorDocumentation/Public/DocumentationCommands.h
deleted file mode 100644
diff --git a/Engine/Plugins/Experimental/InEditorDocumentation/Source/InEditorDocumentation/Public/InEditorDocumentation.h b/Engine/Plugins/Experimental/InEditorDocumentation/Source/InEditorDocumentation/Public/InEditorDocumentation.h
deleted file mode 100644
diff --git a/Engine/Plugins/Experimental/InEditorDocumentation/Source/InEditorDocumentation/Public/InEditorDocumentationSettings.h b/Engine/Plugins/Experimental/InEditorDocumentation/Source/InEditorDocumentation/Public/InEditorDocumentationSettings.h
deleted file mode 100644
diff --git a/Templates/TP_UEIntro_BP/UEIntro.uproject b/Templates/TP_UEIntro_BP/UEIntro.uproject
--- a/Templates/TP_UEIntro_BP/UEIntro.uproject
+++ b/Templates/TP_UEIntro_BP/UEIntro.uproject
@@ -14,16 +14,6 @@
 		{
 			"Name": "Landmass",
 			"Enabled": true
-		},
-		{
-			"Name": "InEditorDocumentation",
-			"Enabled": true,
-			"TargetAllowList": [
-				"Editor"
-			],
-			"PlatformAllowList": [
-				"Win64"
-			]
 		}
 	]
 }
\ No newline at end of file
