Added feedback from QA
Added missing locmetafile
Updated images for template selector

diff --git a/Templates/TP_UEIntro_BP/UEIntro.uproject b/Templates/TP_UEIntro_BP/UEIntro.uproject
--- a/Templates/TP_UEIntro_BP/UEIntro.uproject
+++ b/Templates/TP_UEIntro_BP/UEIntro.uproject
@@ -24,6 +24,10 @@
 			"PlatformAllowList": [
 				"Win64"
 			]
+		},
+		{
+			"Name": "AIAssistant",
+			"Enabled": true
 		}
 	]
 }
\ No newline at end of file
