Bug Fix

Add InEditorDocumentation plugin support for MacOS and Linux

diff --git a/Templates/TP_UEIntro_BP/UEIntro.uproject b/Templates/TP_UEIntro_BP/UEIntro.uproject
--- a/Templates/TP_UEIntro_BP/UEIntro.uproject
+++ b/Templates/TP_UEIntro_BP/UEIntro.uproject
@@ -22,7 +22,9 @@
 				"Editor"
 			],
 			"PlatformAllowList": [
-				"Win64"
+				"Win64",
+				"Mac",
+				"Linux"
 			]
 		},
 		{
